Add "ApplicationCore" as a dependency on the OpenXRHandTracking.build.cs because it calls OnController functions on the message handler, which were inlined before but now cause non-unity errors

#rnx

diff --git a/Engine/Plugins/Runtime/OpenXRHandTracking/Source/OpenXRHandTracking/OpenXRHandTracking.Build.cs b/Engine/Plugins/Runtime/OpenXRHandTracking/Source/OpenXRHandTracking/OpenXRHandTracking.Build.cs
--- a/Engine/Plugins/Runtime/OpenXRHandTracking/Source/OpenXRHandTracking/OpenXRHandTracking.Build.cs
+++ b/Engine/Plugins/Runtime/OpenXRHandTracking/Source/OpenXRHandTracking/OpenXRHandTracking.Build.cs
@@ -38,6 +38,7 @@ namespace UnrealBuildTool.Rules
 					"OpenXRInput",
 					"Slate",
 					"SlateCore",
+					"ApplicationCore"
 				}
 				);
 
