FOpenXRHandTrackingModule  doesn't need to implement IOpenXRExtensionPlugin

FOpenXRHandTracking already implements IOpenXRExtensionPlugin, so we don't need the replicated functionality.
Additionally, the only functionality we used from IOpenXRExtensionPlugin (RegisterOpenXRExtensionModularFeature) was already supported by IOpenXRHandTrackingModule::StartupModule()

diff --git a/Engine/Plugins/Runtime/OpenXRHandTracking/Source/OpenXRHandTracking/Private/OpenXRHandTracking.cpp b/Engine/Plugins/Runtime/OpenXRHandTracking/Source/OpenXRHandTracking/Private/OpenXRHandTracking.cpp
--- a/Engine/Plugins/Runtime/OpenXRHandTracking/Source/OpenXRHandTracking/Private/OpenXRHandTracking.cpp
+++ b/Engine/Plugins/Runtime/OpenXRHandTracking/Source/OpenXRHandTracking/Private/OpenXRHandTracking.cpp
@@ -31,8 +31,7 @@ static_assert(EHandKeypoint::LittleTip == static_cast<EHandKeypoint>(XR_HAND_JOI
 //static TAutoConsoleVariable<int32> CVarEnableOpenXRHandTrackingDebug(TEXT("OpenXR.debug.EnableEyetrackingDebug"), 1, TEXT("0 - Eyetracking debug visualizations are disabled. 1 - Eyetracking debug visualizations are enabled."));
 
 class FOpenXRHandTrackingModule :
-	public IOpenXRHandTrackingModule,
-	public IOpenXRExtensionPlugin
+	public IOpenXRHandTrackingModule
 {
 public:
 	FOpenXRHandTrackingModule()
@@ -43,7 +42,6 @@ public:
 	virtual void StartupModule() override
 	{
 		IOpenXRHandTrackingModule::StartupModule();
-		RegisterOpenXRExtensionModularFeature();
 
 		// HACK: Generic Application might not be instantiated at this point so we create the input device with a
 		// dummy message handler. When the Generic Application creates the input device it passes a valid message
