OpenXRHandTracking
-Better support for using hand tracking along with controllers, rather than one or the other.
-Two new project settings for the OpenXRHandTracking plugin:
-bUseMoreSpecificMotionSourceNames change the hand tracking keypoint motion source names from [Left|Right][Keypoint] to HandTracking[Left|Right][Keypoint] (default false)
-bSupportLegacyControllerMotionSources if true the legacy motion sources Left, Right, and AnyHand can be supplied by hand tracking.  If False they cannot.
-Note the Left/Right Grip/Aim poses already worked only for the controllers, and specific keypoint sources existed for every hand bone so it was possible to do this before, but this makes it less confusing in some cases.
-Refactored hand tracking keypoint motion source pose collection to use a map of motion source fname to keypoint enum value rather than turning both into strings and comparing every time they are accessed.

diff --git a/Engine/Plugins/Runtime/OpenXRHandTracking/OpenXRHandTracking.uplugin b/Engine/Plugins/Runtime/OpenXRHandTracking/OpenXRHandTracking.uplugin
--- a/Engine/Plugins/Runtime/OpenXRHandTracking/OpenXRHandTracking.uplugin
+++ b/Engine/Plugins/Runtime/OpenXRHandTracking/OpenXRHandTracking.uplugin
@@ -29,6 +29,15 @@
 				"Linux",
 				"Android"
 			]
+		},
+		{
+			"Name": "OpenXRHandTrackingEditor",
+			"Type": "Editor",
+			"LoadingPhase": "Default",
+			"PlatformAllowList": [
+				"Win64",
+				"Linux"
+			]
 		}
 	],
 	"Plugins": [
diff --git a/Engine/Plugins/Runtime/OpenXRHandTracking/Source/OpenXRHandTracking/OpenXRHandTracking.Build.cs b/Engine/Plugins/Runtime/OpenXRHandTracking/Source/OpenXRHandTracking/OpenXRHandTracking.Build.cs
--- a/Engine/Plugins/Runtime/OpenXRHandTracking/Source/OpenXRHandTracking/OpenXRHandTracking.Build.cs
+++ b/Engine/Plugins/Runtime/OpenXRHandTracking/Source/OpenXRHandTracking/OpenXRHandTracking.Build.cs
@@ -31,7 +31,9 @@ namespace UnrealBuildTool.Rules
 
             if (Target.bBuildEditor == true)
             {
-                PrivateDependencyModuleNames.Add("UnrealEd");
+				PrivateDependencyModuleNames.Add("EditorFramework");
+				PrivateDependencyModuleNames.Add("UnrealEd");
+				PrivateDependencyModuleNames.Add("InputEditor");
             }
         }
     }
diff --git a/Engine/Plugins/Runtime/OpenXRHandTracking/Source/OpenXRHandTracking/Private/OpenXRHandTracking.cpp b/Engine/Plugins/Runtime/OpenXRHandTracking/Source/OpenXRHandTracking/Private/OpenXRHandTracking.cpp
--- a/Engine/Plugins/Runtime/OpenXRHandTracking/Source/OpenXRHandTracking/Private/OpenXRHandTracking.cpp
+++ b/Engine/Plugins/Runtime/OpenXRHandTracking/Source/OpenXRHandTracking/Private/OpenXRHandTracking.cpp
@@ -10,6 +10,8 @@
 //#include "DrawDebugHelpers.h"
 #include "ILiveLinkClient.h"
 #include "IOpenXRHMDModule.h"
+#include "OpenXRHandTrackingSettings.h"
+#include "UObject/UObjectGlobals.h"
 
 #define LOCTEXT_NAMESPACE "OpenXRHandTracking"
 
@@ -20,8 +22,6 @@ static_assert(EHandKeypoint::ThumbMetacarpal == static_cast<EHandKeypoint>(XR_HA
 static_assert(EHandKeypoint::IndexTip == static_cast<EHandKeypoint>(XR_HAND_JOINT_INDEX_TIP_EXT), "EHandKeypoint enum does not match XrHandJointEXT.");
 static_assert(EHandKeypoint::LittleTip == static_cast<EHandKeypoint>(XR_HAND_JOINT_LITTLE_TIP_EXT), "EHandKeypoint enum does not match XrHandJointEXT.");
 
-//static TAutoConsoleVariable<int32> CVarEnableOpenXRHandTrackingDebug(TEXT("OpenXR.debug.EnableEyetrackingDebug"), 1, TEXT("0 - Eyetracking debug visualizations are disabled. 1 - Eyetracking debug visualizations are enabled."));
-
 class FOpenXRHandTrackingModule :
 	public IOpenXRHandTrackingModule
 {
@@ -197,6 +197,9 @@ const void* FOpenXRHandTracking::OnGetSystem(XrInstance InInstance, const void*
 
 const void* FOpenXRHandTracking::OnCreateSession(XrInstance InInstance, XrSystemId InSystem, const void* InNext)
 {
+	// Need to wait until the EHandKeypoint enum has been loaded, so we do this here rather than in the constructor which runs too early.
+	BuildMotionSourceToKeypointMap();
+
 	XrSystemHandTrackingPropertiesEXT HandTrackingSystemProperties{
 	XR_TYPE_SYSTEM_HAND_TRACKING_PROPERTIES_EXT };
 	XrSystemProperties systemProperties{ XR_TYPE_SYSTEM_PROPERTIES,
@@ -300,61 +303,83 @@ const FTransform& FOpenXRHandTracking::FHandState::GetTransform(EHandKeypoint Ke
 	return KeypointTransforms[(uint32)Keypoint];
 }
 
+void FOpenXRHandTracking::BuildMotionSourceToKeypointMap()
+{
+	if (!MotionSourceToKeypointMap.IsEmpty())
+	{
+		return;
+	}
+
+	bool bUseMoreSpecificMotionSourceNames = false;
+	UOpenXRHandTrackingSettings* Settings = GetMutableDefault<UOpenXRHandTrackingSettings>();
+	ensure(Settings);
+	if (Settings)
+	{
+		bUseMoreSpecificMotionSourceNames = Settings->bUseMoreSpecificMotionSourceNames;
+		bSupportLegacyControllerMotionSources = Settings->bSupportLegacyControllerMotionSources;
+	}
+
+	// There is a motionsource that corresponds to each hand keypoint of the form [Left|Right][Keypoint].
+	// Build a map so we can quickly translate from motion source FName to hand bone EHandKeypoint value.
+	// We also have the option of using more specific motion sources of the form HandTracking[Left|Right][Keypoint]
+	// this is useful if one wishes to use hand tracking and controllers simultaneously.
+	// We also may support more generic legacy motion sources, by default we do support this.
+	const UEnum* EnumPtr = FindObject<UEnum>(nullptr, TEXT("/Script/HeadMountedDisplay.EHandKeypoint"), true);
+	check(EnumPtr != nullptr);
+
+	check(IsInGameThread());
+	const FString Left(bUseMoreSpecificMotionSourceNames ? TEXT("HandTrackingLeft") : TEXT("Left"));
+	const FString Right(bUseMoreSpecificMotionSourceNames ? TEXT("HandTrackingRight") : TEXT("Right"));
+	for (int64 E = 0; E < EHandKeypointCount; ++E)
+	{
+		const EHandKeypoint EnumValue = static_cast<EHandKeypoint>(E);
+		const FString EnumName = EnumPtr->GetNameStringByValue(E);
+
+		const FName LeftName(Left + EnumName);
+		const FName RightName(Right + EnumName);
+		MotionSourceToKeypointMap.Add(LeftName, MotionSourceInfo(EnumValue, true));
+		MotionSourceToKeypointMap.Add(RightName, MotionSourceInfo(EnumValue, false));
+	}
+
+	if (bSupportLegacyControllerMotionSources)
+	{
+		MotionSourceToKeypointMap.Add(FName("Left"), MotionSourceInfo(EHandKeypoint::Palm, true));
+		MotionSourceToKeypointMap.Add(FName("Right"), MotionSourceInfo(EHandKeypoint::Palm, false));
+	}
+}
+
 bool FOpenXRHandTracking::GetControllerOrientationAndPosition(const int32 ControllerIndex, const FName MotionSource, FRotator& OutOrientation, FVector& OutPosition, float WorldToMetersScale) const
 {
 	if (!bHandTrackingAvailable)
 	{
 		return false;
 	}
+	
+	// Hand tracking currently does not support late update.  Current hand tracking systems have latency that make it pointless.
+	if (!IsInGameThread())
+	{
+		return false;
+	}
 
 	bool bTracked = false;
 	if (ControllerIndex == DeviceIndex)
 	{
 		FTransform ControllerTransform = FTransform::Identity;
-		if (MotionSource == FName("Left"))
-		{
-			ControllerTransform = GetLeftHandState().GetTransform(EHandKeypoint::Palm);
-			bTracked = GetLeftHandState().ReceivedJointPoses;
-		}
-		else if (MotionSource == FName("Right"))
-		{
-			ControllerTransform = GetRightHandState().GetTransform(EHandKeypoint::Palm);
-			bTracked = GetRightHandState().ReceivedJointPoses;
-		}
 
-		// This can only be done in the game thread since it uses the UEnum directly
-		if (IsInGameThread())
+		const MotionSourceInfo* KeyPointInfoPtr = MotionSourceToKeypointMap.Find(MotionSource);
+		if (KeyPointInfoPtr)
 		{
-			const UEnum* EnumPtr = FindObject<UEnum>(nullptr, TEXT("/Script/HeadMountedDisplay.EHandKeypoint"), true);
-			check(EnumPtr != nullptr);
-			bool bUseRightHand = false;
-			FString SourceString = MotionSource.ToString();
-			if (SourceString.StartsWith((TEXT("Right"))))
+			const EHandKeypoint KeyPoint = KeyPointInfoPtr->Key;
+			const bool bIsLeft = KeyPointInfoPtr->Value;
+			if (bIsLeft)
 			{
-				bUseRightHand = true;
-				// Strip off the Right
-				SourceString.RightInline(SourceString.Len() - 5, false);
+				ControllerTransform = GetLeftHandState().GetTransform(KeyPoint);
+				bTracked = GetLeftHandState().ReceivedJointPoses;
 			}
 			else
 			{
-				// Strip off the Left
-				SourceString.RightInline(SourceString.Len() - 4, false);
-			}
-			FName FullEnumName(*FString(TEXT("EHandKeypoint::") + SourceString), FNAME_Find);
-			// Get the enum value from the name
-			int32 ValueFromName = EnumPtr->GetValueByName(FullEnumName);
-			if (ValueFromName != INDEX_NONE)
-			{
-				if (bUseRightHand)
-				{
-					ControllerTransform = GetRightHandState().GetTransform((EHandKeypoint)ValueFromName);
-					bTracked = GetRightHandState().ReceivedJointPoses;
-				}
-				else
-				{
-					ControllerTransform = GetLeftHandState().GetTransform((EHandKeypoint)ValueFromName);
-					bTracked = GetLeftHandState().ReceivedJointPoses;
-				}
+				ControllerTransform = GetRightHandState().GetTransform(KeyPoint);
+				bTracked = GetRightHandState().ReceivedJointPoses;
 			}
 		}
 
@@ -362,8 +387,11 @@ bool FOpenXRHandTracking::GetControllerOrientationAndPosition(const int32 Contro
 		OutOrientation = ControllerTransform.GetRotation().Rotator();
 	}
 
-	// Then call super to handle a few of the default labels, for backward compatibility
-	FXRMotionControllerBase::GetControllerOrientationAndPosition(ControllerIndex, MotionSource, OutOrientation, OutPosition, WorldToMetersScale);
+	// Then call super to handle a few of the default labels (eg AnyHand), for backward compatibility
+	if (bSupportLegacyControllerMotionSources)
+	{
+		FXRMotionControllerBase::GetControllerOrientationAndPosition(ControllerIndex, MotionSource, OutOrientation, OutPosition, WorldToMetersScale);
+	}
 
 	return bTracked;
 }
@@ -425,10 +453,20 @@ void FOpenXRHandTracking::EnumerateSources(TArray<FMotionControllerSource>& Sour
 {
 	check(IsInGameThread());
 
-	SourcesOut.Empty(EHandKeypointCount * 2);
+	bool bUseMoreSpecificMotionSourceNames = false;
+	UOpenXRHandTrackingSettings* Settings = GetMutableDefault<UOpenXRHandTrackingSettings>();
+	ensure(Settings);
+	if (Settings)
+	{
+		bUseMoreSpecificMotionSourceNames = Settings->bUseMoreSpecificMotionSourceNames;
+	}
+
+	SourcesOut.Reserve(SourcesOut.Num() + (EHandKeypointCount * 2));
 
 	const UEnum* EnumPtr = FindObject<UEnum>(nullptr, TEXT("/Script/HeadMountedDisplay.EHandKeypoint"), true);
 	check(EnumPtr != nullptr);
+	const FString Left(bUseMoreSpecificMotionSourceNames ? TEXT("HandTrackingLeft") : TEXT("Left"));
+	const FString Right(bUseMoreSpecificMotionSourceNames ? TEXT("HandTrackingRight") : TEXT("Right"));
 	for (int32 Keypoint = 0; Keypoint < EHandKeypointCount; Keypoint++)
 	{
 		static int32 EnumNameLength = FString(TEXT("EHandKeypoint::")).Len();
@@ -436,9 +474,9 @@ void FOpenXRHandTracking::EnumerateSources(TArray<FMotionControllerSource>& Sour
 		const FString EnumString = EnumPtr->GetNameByValue(Keypoint).ToString();
 		const TCHAR* EnumChars = *EnumString;
 		const TCHAR* EnumValue = EnumChars + EnumNameLength;
-		FString StringLeft(TEXT("Left"));
+		FString StringLeft(Left);
 		StringLeft.AppendChars(EnumValue, EnumString.Len() - EnumNameLength);
-		FString StringRight(TEXT("Right"));
+		FString StringRight(Right);
 		StringRight.AppendChars(EnumValue, EnumString.Len() - EnumNameLength);
 		FName SourceL(*(StringLeft));
 		FName SourceR(*(StringRight));
diff --git a/Engine/Plugins/Runtime/OpenXRHandTracking/Source/OpenXRHandTracking/Private/OpenXRHandTracking.h b/Engine/Plugins/Runtime/OpenXRHandTracking/Source/OpenXRHandTracking/Private/OpenXRHandTracking.h
--- a/Engine/Plugins/Runtime/OpenXRHandTracking/Source/OpenXRHandTracking/Private/OpenXRHandTracking.h
+++ b/Engine/Plugins/Runtime/OpenXRHandTracking/Source/OpenXRHandTracking/Private/OpenXRHandTracking.h
@@ -108,6 +108,7 @@ public:
 
 private:
 	void AddKeys();
+	void BuildMotionSourceToKeypointMap();
 	
 	void SetupLiveLinkData();
 	void UpdateLiveLink();
@@ -128,6 +129,9 @@ private:
 
 	TArray<int32> BoneParents;
 	TArray<EHandKeypoint> BoneKeypoints;
+	typedef TPair<EHandKeypoint, bool> MotionSourceInfo; // bool true == left, false == right
+	TMap<FName, MotionSourceInfo> MotionSourceToKeypointMap;
+	bool bSupportLegacyControllerMotionSources = true;
 
 	FHandState HandStates[2];
 
diff --git a/Engine/Plugins/Runtime/OpenXRHandTracking/Source/OpenXRHandTracking/Private/OpenXRHandTrackingSettings.cpp b/Engine/Plugins/Runtime/OpenXRHandTracking/Source/OpenXRHandTracking/Private/OpenXRHandTrackingSettings.cpp
new file mode 100644
--- /dev/null
+++ b/Engine/Plugins/Runtime/OpenXRHandTracking/Source/OpenXRHandTracking/Private/OpenXRHandTrackingSettings.cpp
@@ -0,0 +1,15 @@
+// Copyright Epic Games, Inc. All Rights Reserved.
+
+#include "OpenXRHandTrackingSettings.h"
+//#include "XRMotionControllerBase.h"
+//#include "PlayerMappableInputConfig.h"
+
+//#if WITH_EDITOR
+//#include "Editor.h"
+//#include "Features/IModularFeatures.h"
+//#endif
+
+UOpenXRHandTrackingSettings::UOpenXRHandTrackingSettings(const FObjectInitializer& ObjectInitializer)
+	: Super(ObjectInitializer)
+{
+}
diff --git a/Engine/Plugins/Runtime/OpenXRHandTracking/Source/OpenXRHandTracking/Public/OpenXRHandTrackingSettings.h b/Engine/Plugins/Runtime/OpenXRHandTracking/Source/OpenXRHandTracking/Public/OpenXRHandTrackingSettings.h
new file mode 100644
--- /dev/null
+++ b/Engine/Plugins/Runtime/OpenXRHandTracking/Source/OpenXRHandTracking/Public/OpenXRHandTrackingSettings.h
@@ -0,0 +1,25 @@
+// Copyright Epic Games, Inc. All Rights Reserved.
+
+#pragma once
+
+#include "CoreMinimal.h"
+#include "UObject/Object.h"
+#include "OpenXRHandTrackingSettings.generated.h"
+
+/**
+* Implements the settings for the OpenXR Input plugin.
+*/
+UCLASS(config = Input, defaultconfig)
+class OPENXRHANDTRACKING_API UOpenXRHandTrackingSettings : public UObject
+{
+	GENERATED_UCLASS_BODY()
+
+public:
+	/** If false (the default) the motion sources for hand tracking will be of the form '[Left|Right][Keypoint]'.  If true they will be of the form 'HandTracking[Left|Right][Keypoint]'.  True is reccomended to avoid collisions between motion sources from different device types. **/
+	UPROPERTY(config, EditAnywhere, Category = "Motion Sources")
+	bool bUseMoreSpecificMotionSourceNames = false;
+
+	/** If true hand tracking supports the 'Left' and 'Right' legacy motion sources.  If false it does not.  False is reccomended unless you need legacy compatibility in an older unreal projects.**/
+	UPROPERTY(config, EditAnywhere, Category = "Motion Sources")
+	bool bSupportLegacyControllerMotionSources = true;
+};
diff --git a/Engine/Plugins/Runtime/OpenXRHandTracking/Source/OpenXRHandTrackingEditor/OpenXRHandTrackingEditor.Build.cs b/Engine/Plugins/Runtime/OpenXRHandTracking/Source/OpenXRHandTrackingEditor/OpenXRHandTrackingEditor.Build.cs
new file mode 100644
--- /dev/null
+++ b/Engine/Plugins/Runtime/OpenXRHandTracking/Source/OpenXRHandTrackingEditor/OpenXRHandTrackingEditor.Build.cs
@@ -0,0 +1,19 @@
+// Copyright Epic Games, Inc. All Rights Reserved.
+
+using UnrealBuildTool;
+
+public class OpenXRHandTrackingEditor : ModuleRules
+{
+	public OpenXRHandTrackingEditor(ReadOnlyTargetRules Target) : base(Target)
+	{
+		PrivateDependencyModuleNames.AddRange(
+			new string[] {
+				"Core",
+				"CoreUObject",
+				"EditorFramework",
+				"UnrealEd",
+				"OpenXRHandTracking"
+			}
+		);
+	}
+}
diff --git a/Engine/Plugins/Runtime/OpenXRHandTracking/Source/OpenXRHandTrackingEditor/Private/OpenXRHandTrackingEditorModule.cpp b/Engine/Plugins/Runtime/OpenXRHandTracking/Source/OpenXRHandTrackingEditor/Private/OpenXRHandTrackingEditorModule.cpp
new file mode 100644
--- /dev/null
+++ b/Engine/Plugins/Runtime/OpenXRHandTracking/Source/OpenXRHandTrackingEditor/Private/OpenXRHandTrackingEditorModule.cpp
@@ -0,0 +1,63 @@
+// Copyright Epic Games, Inc. All Rights Reserved.
+
+#include "OpenXRHandTrackingSettings.h"
+#include "ISettingsModule.h"
+#include "Modules/ModuleManager.h"
+
+
+#define LOCTEXT_NAMESPACE "FOpenXRHandTrackingEditorModule"
+
+/**
+ * Module for the HoloLens platform editor module.
+ */
+class FOpenXRHandTrackingEditorModule
+	: public IModuleInterface
+{
+public:
+
+	/** Default constructor. */
+	FOpenXRHandTrackingEditorModule( )
+	{ }
+
+	/** Destructor. */
+	~FOpenXRHandTrackingEditorModule( )
+	{
+	}
+
+public:
+
+	// IModuleInterface interface
+
+	virtual void StartupModule() override
+	{
+		ISettingsModule* SettingsModule = FModuleManager::GetModulePtr<ISettingsModule>("Settings");
+
+		if (SettingsModule != nullptr)
+		{
+			SettingsModule->RegisterSettings("Project", "Plugins", "OpenXRHandTracking",
+				LOCTEXT("SettingsName", "OpenXRHandTracking"),
+				LOCTEXT("SettingsDescription", "Settings for OpenXRHandTracking Plugin"),
+				GetMutableDefault<UOpenXRHandTrackingSettings>()
+			);
+		}
+	}
+
+	virtual void ShutdownModule() override
+	{
+		ISettingsModule* SettingsModule = FModuleManager::GetModulePtr<ISettingsModule>("Settings");
+		if (SettingsModule != nullptr)
+		{
+			SettingsModule->UnregisterSettings("Project", "Plugins", "OpenXRHandTracking");
+		}
+	}
+
+private:
+
+};
+
+
+#undef LOCTEXT_NAMESPACE
+
+
+IMPLEMENT_MODULE(FOpenXRHandTrackingEditorModule, OpenXRHandTrackingEditorModule);
+
