OpenXR: fix OpenXRHandTrackingEditor module name

The plugin descriptor and module had different names for the module. This affected the CookedEditor build's ability to load the module.

#rnx

diff --git a/Engine/Plugins/Runtime/OpenXRHandTracking/Source/OpenXRHandTrackingEditor/Private/OpenXRHandTrackingEditorModule.cpp b/Engine/Plugins/Runtime/OpenXRHandTracking/Source/OpenXRHandTrackingEditor/Private/OpenXRHandTrackingEditorModule.cpp
--- a/Engine/Plugins/Runtime/OpenXRHandTracking/Source/OpenXRHandTrackingEditor/Private/OpenXRHandTrackingEditorModule.cpp
+++ b/Engine/Plugins/Runtime/OpenXRHandTracking/Source/OpenXRHandTrackingEditor/Private/OpenXRHandTrackingEditorModule.cpp
@@ -59,5 +59,5 @@ private:
 #undef LOCTEXT_NAMESPACE
 
 
-IMPLEMENT_MODULE(FOpenXRHandTrackingEditorModule, OpenXRHandTrackingEditorModule);
+IMPLEMENT_MODULE(FOpenXRHandTrackingEditorModule, OpenXRHandTrackingEditor);
 
