Added OnCreateHandTracker() IOpenXRExtensionPlugin hook for xrCreateHandTrackerEXT() OpenXR call

diff --git a/Engine/Plugins/Runtime/OpenXR/Source/OpenXRHMD/Public/IOpenXRExtensionPlugin.h b/Engine/Plugins/Runtime/OpenXR/Source/OpenXRHMD/Public/IOpenXRExtensionPlugin.h
--- a/Engine/Plugins/Runtime/OpenXR/Source/OpenXRHMD/Public/IOpenXRExtensionPlugin.h
+++ b/Engine/Plugins/Runtime/OpenXR/Source/OpenXRHMD/Public/IOpenXRExtensionPlugin.h
@@ -565,4 +565,6 @@ public:
 	/** Stop the AR system and reset its internal state; this task must succeed. */
 	virtual void OnStopARSession() {}
 
+	virtual const void* OnCreateHandTracker(const void* InNext) { return InNext;}
+	
 };
diff --git a/Engine/Plugins/Runtime/OpenXRHandTracking/Source/OpenXRHandTracking/Private/OpenXRHandTracking.cpp b/Engine/Plugins/Runtime/OpenXRHandTracking/Source/OpenXRHandTracking/Private/OpenXRHandTracking.cpp
--- a/Engine/Plugins/Runtime/OpenXRHandTracking/Source/OpenXRHandTracking/Private/OpenXRHandTracking.cpp
+++ b/Engine/Plugins/Runtime/OpenXRHandTracking/Source/OpenXRHandTracking/Private/OpenXRHandTracking.cpp
@@ -9,6 +9,7 @@
 #include "Framework/Application/SlateApplication.h"
 //#include "DrawDebugHelpers.h"
 #include "ILiveLinkClient.h"
+#include "IOpenXRHMD.h"
 #include "IOpenXRHMDModule.h"
 #include "OpenXRHandTrackingSettings.h"
 #include "UObject/UObjectGlobals.h"
@@ -173,7 +174,7 @@ bool FOpenXRHandTracking::GetRequiredExtensions(TArray<const ANSICHAR*>& OutExte
 		}
 	}
 
-	OutExtensions.Add("XR_EXT_hand_tracking");
+	OutExtensions.Add(XR_EXT_HAND_TRACKING_EXTENSION_NAME);
 	return true;
 }
 
@@ -215,21 +216,11 @@ const void* FOpenXRHandTracking::OnBeginSession(XrSession InSession, const void*
 
 		// Create a hand tracker for left hand that tracks default set of hand joints.
 		FHandState& LeftHandState = GetLeftHandState();
-		{
-			XrHandTrackerCreateInfoEXT CreateInfo{ XR_TYPE_HAND_TRACKER_CREATE_INFO_EXT };
-			CreateInfo.hand = XR_HAND_LEFT_EXT;
-			CreateInfo.handJointSet = XR_HAND_JOINT_SET_DEFAULT_EXT;
-			XR_ENSURE(xrCreateHandTrackerEXT(InSession, &CreateInfo, &LeftHandState.HandTracker));
-		}
+		CreateHandTracker(InSession, XR_HAND_LEFT_EXT, LeftHandState);
 
-		// Create a hand tracker for left hand that tracks default set of hand joints.
+		// Create a hand tracker for right  hand that tracks default set of hand joints.
 		FHandState& RightHandState = GetRightHandState();
-		{
-			XrHandTrackerCreateInfoEXT CreateInfo{ XR_TYPE_HAND_TRACKER_CREATE_INFO_EXT };
-			CreateInfo.hand = XR_HAND_RIGHT_EXT;
-			CreateInfo.handJointSet = XR_HAND_JOINT_SET_DEFAULT_EXT;
-			XR_ENSURE(xrCreateHandTrackerEXT(InSession, &CreateInfo, &RightHandState.HandTracker));
-		}
+		CreateHandTracker(InSession, XR_HAND_RIGHT_EXT, RightHandState);
 	}
 
 	return InNext;
@@ -575,4 +566,23 @@ void FOpenXRHandTracking::AddKeys()
 {
 }
 
+void FOpenXRHandTracking::CreateHandTracker(XrSession InSession, XrHandEXT HandEXT, FHandState& HandState)
+{
+	XrHandTrackerCreateInfoEXT CreateInfo{ XR_TYPE_HAND_TRACKER_CREATE_INFO_EXT };
+	CreateInfo.hand = HandEXT;
+	CreateInfo.handJointSet = XR_HAND_JOINT_SET_DEFAULT_EXT;
+
+	if (XRTrackingSystem && XRTrackingSystem->GetIOpenXRHMD())
+	{
+		TArray<IOpenXRExtensionPlugin*> ExtensionPlugins;
+		ExtensionPlugins = XRTrackingSystem->GetIOpenXRHMD()->GetExtensionPlugins();
+		for (IOpenXRExtensionPlugin* Plugin : ExtensionPlugins)
+		{
+			CreateInfo.next = Plugin->OnCreateHandTracker(CreateInfo.next);
+		}
+	}
+	
+	XR_ENSURE(xrCreateHandTrackerEXT(InSession, &CreateInfo, &HandState.HandTracker));
+}
+
 #undef LOCTEXT_NAMESPACE
diff --git a/Engine/Plugins/Runtime/OpenXRHandTracking/Source/OpenXRHandTracking/Private/OpenXRHandTracking.h b/Engine/Plugins/Runtime/OpenXRHandTracking/Source/OpenXRHandTracking/Private/OpenXRHandTracking.h
--- a/Engine/Plugins/Runtime/OpenXRHandTracking/Source/OpenXRHandTracking/Private/OpenXRHandTracking.h
+++ b/Engine/Plugins/Runtime/OpenXRHandTracking/Source/OpenXRHandTracking/Private/OpenXRHandTracking.h
@@ -113,6 +113,7 @@ private:
 	void SetupLiveLinkData();
 	void UpdateLiveLink();
 	void UpdateLiveLinkTransforms(TArray<FTransform>& OutTransforms, const FOpenXRHandTracking::FHandState& HandState);
+	void CreateHandTracker(XrSession InSession, XrHandEXT HandEXT, FHandState& HandState);
 
 	bool bHandTrackingAvailable = false;
 
